Fixed SoundCueEditor link errors

diff --git a/Engine/Source/Editor/SoundCueEditor/SoundCueEditor.Build.cs b/Engine/Source/Editor/SoundCueEditor/SoundCueEditor.Build.cs
--- a/Engine/Source/Editor/SoundCueEditor/SoundCueEditor.Build.cs
+++ b/Engine/Source/Editor/SoundCueEditor/SoundCueEditor.Build.cs
@@ -10,6 +10,7 @@ public class SoundCueEditor : ModuleRules
 			new string[] {
 				"Core",
 				"CoreUObject",
+                "InputCore",
 				"Engine",
 				"Slate",
 				"SlateCore",
