UE-12177 Fix to stop sound cues playing when the sound cue editor closes

https://jira.ol.epicgames.net/browse/UE-12177

diff --git a/Engine/Source/Editor/SoundCueEditor/Private/SoundCueEditor.cpp b/Engine/Source/Editor/SoundCueEditor/Private/SoundCueEditor.cpp
--- a/Engine/Source/Editor/SoundCueEditor/Private/SoundCueEditor.cpp
+++ b/Engine/Source/Editor/SoundCueEditor/Private/SoundCueEditor.cpp
@@ -60,6 +60,13 @@ void FSoundCueEditor::UnregisterTabSpawners(const TSharedRef<class FTabManager>&
 
 FSoundCueEditor::~FSoundCueEditor()
 {
+	// Stop any playing sound cues when the cue editor closes
+	UAudioComponent* PreviewComp = GEditor->GetPreviewAudioComponent();
+	if (PreviewComp && PreviewComp->IsPlaying())
+	{
+		Stop();
+	}
+
 	GEditor->UnregisterForUndo( this );
 }
 
